Added PixelStreaming backend option to VCam.

diff --git a/Engine/Plugins/Media/PixelStreamingVCam/PixelStreamingVCam.uplugin b/Engine/Plugins/Media/PixelStreamingVCam/PixelStreamingVCam.uplugin
new file mode 100644
--- /dev/null
+++ b/Engine/Plugins/Media/PixelStreamingVCam/PixelStreamingVCam.uplugin
@@ -0,0 +1,40 @@
+{
+	"FileVersion": 3,
+	"Version": 1.0,
+	"VersionName": "1.0",
+	"FriendlyName": "Pixel Streaming VCam Support",
+	"Description": "Pixel Streaming extensions to VCam.",
+	"Category": "Graphics",
+	"CreatedBy": "Epic Games, Inc.",
+	"CreatedByURL": "http://epicgames.com",
+	"DocsURL": "https://docs.unrealengine.com/en-US/Platforms/PixelStreaming/index.html",
+	"MarketplaceURL": "",
+	"SupportURL": "",
+	"EnabledByDefault": false,
+	"CanContainContent": true,
+	"IsBetaVersion": false,
+	"Installed": false,
+	"Modules": [
+		{
+			"Name": "PixelStreamingVCam",
+			"Type": "Editor",
+			"LoadingPhase": "PostEngineInit",
+			"WhitelistPlatforms": [ "Win64", "Linux" ],
+			"BlacklistTargets": [ "Server" ]
+		}
+	],
+	"Plugins": [
+		{
+			"Name": "PixelStreaming",
+			"Enabled": true
+		},
+		{
+			"Name": "MediaIOFramework",
+			"Enabled": true
+		},
+		{
+			"Name": "VirtualCamera",
+			"Enabled": true
+		}
+	]
+}
diff --git a/Engine/Plugins/Media/PixelStreamingVCam/Source/PixelStreamingVCam/PixelStreamingVCam.Build.cs b/Engine/Plugins/Media/PixelStreamingVCam/Source/PixelStreamingVCam/PixelStreamingVCam.Build.cs
new file mode 100644
--- /dev/null
+++ b/Engine/Plugins/Media/PixelStreamingVCam/Source/PixelStreamingVCam/PixelStreamingVCam.Build.cs
@@ -0,0 +1,40 @@
+// Copyright Epic Games, Inc. All Rights Reserved.
+
+using System;
+using System.IO;
+using System.Collections.Generic;
+
+namespace UnrealBuildTool.Rules
+{
+	public class PixelStreamingVCam : ModuleRules
+	{
+		public PixelStreamingVCam(ReadOnlyTargetRules Target) : base(Target)
+		{
+			// This is so for game projects using our public headers don't have to include extra modules they might not know about.
+			PublicDependencyModuleNames.AddRange(new string[]
+			{
+				"InputDevice",
+				"MediaIOCore",
+				"VCamCore",
+				"PixelStreaming"
+			});
+
+			PrivateDependencyModuleNames.AddRange(new string[]
+			{
+				"ApplicationCore",
+				"Core",
+				"CoreUObject",
+				"Engine",
+				"UMG",
+				"RHI"
+			});
+
+			// Can't package non-editor targets (e.g. games) with UnrealEd, so this dependency should only be added in editor.
+			if (Target.bBuildEditor)
+			{
+				PrivateDependencyModuleNames.Add("UnrealEd");
+			}
+
+		}
+	}
+}
diff --git a/Engine/Plugins/Media/PixelStreamingVCam/Source/PixelStreamingVCam/Private/PixelStreamingMediaCapture.cpp b/Engine/Plugins/Media/PixelStreamingVCam/Source/PixelStreamingVCam/Private/PixelStreamingMediaCapture.cpp
new file mode 100644
--- /dev/null
+++ b/Engine/Plugins/Media/PixelStreamingVCam/Source/PixelStreamingVCam/Private/PixelStreamingMediaCapture.cpp
@@ -0,0 +1,45 @@
+// Copyright Epic Games, Inc. All Rights Reserved.
+
+#include "PixelStreamingMediaCapture.h"
+#include "PixelStreamingSourceFrame.h"
+#include "Slate/SceneViewport.h"
+
+void UPixelStreamingMediaCapture::OnCustomCapture_RenderingThread(
+	FRHICommandListImmediate& RHICmdList,
+	const FCaptureBaseData& InBaseData,
+	TSharedPtr<FMediaCaptureUserData, ESPMode::ThreadSafe> InUserData,
+	FTexture2DRHIRef InSourceTexture,
+	FTextureRHIRef TargetableTexture,
+	FResolveParams& ResolveParams,
+	FVector2D CropU,
+	FVector2D CropV)
+{
+	if (VideoInput)
+	{
+		VideoInput->OnFrame.Broadcast(FPixelStreamingSourceFrame(InSourceTexture));
+	}
+}
+
+bool UPixelStreamingMediaCapture::CaptureSceneViewportImpl(TSharedPtr<FSceneViewport>& InSceneViewport)
+{
+	Viewport = InSceneViewport;
+	SetupVideoInput();
+	SetState(EMediaCaptureState::Capturing);
+	return true;
+}
+
+bool UPixelStreamingMediaCapture::CaptureRenderTargetImpl(UTextureRenderTarget2D* InRenderTarget)
+{
+	Viewport = nullptr;
+	SetupVideoInput();
+	SetState(EMediaCaptureState::Capturing);
+	return true;
+}
+
+void UPixelStreamingMediaCapture::SetupVideoInput()
+{
+	if (!VideoInput)
+	{
+		VideoInput = MakeShared<FPixelStreamingVideoInput>();
+	}
+}
diff --git a/Engine/Plugins/Media/PixelStreamingVCam/Source/PixelStreamingVCam/Private/PixelStreamingMediaCapture.h b/Engine/Plugins/Media/PixelStreamingVCam/Source/PixelStreamingVCam/Private/PixelStreamingMediaCapture.h
new file mode 100644
--- /dev/null
+++ b/Engine/Plugins/Media/PixelStreamingVCam/Source/PixelStreamingVCam/Private/PixelStreamingMediaCapture.h
@@ -0,0 +1,40 @@
+// Copyright Epic Games, Inc. All Rights Reserved.
+
+#pragma once
+
+#include "MediaCapture.h"
+#include "PixelStreamingVideoInput.h"
+#include "PixelStreamingMediaCapture.generated.h"
+
+UCLASS(BlueprintType)
+class UPixelStreamingMediaCapture : public UMediaCapture
+{
+	GENERATED_BODY()
+
+	//~ Begin UMediaCapture interface
+public:
+	virtual void OnCustomCapture_RenderingThread(
+		FRHICommandListImmediate& RHICmdList,
+		const FCaptureBaseData& InBaseData,
+		TSharedPtr<FMediaCaptureUserData, ESPMode::ThreadSafe> InUserData,
+		FTexture2DRHIRef InSourceTexture,
+		FTextureRHIRef TargetableTexture,
+		FResolveParams& ResolveParams,
+		FVector2D CropU,
+		FVector2D CropV) override;
+	virtual bool CaptureSceneViewportImpl(TSharedPtr<FSceneViewport>& InSceneViewport) override;
+	virtual bool CaptureRenderTargetImpl(UTextureRenderTarget2D* InRenderTarget) override;
+	virtual bool ShouldCaptureRHITexture() const { return true; }
+	virtual void StopCaptureImpl(bool bAllowPendingFrameToBeProcess) override {}
+	//~ End UMediaCapture interface
+
+	TSharedPtr<FPixelStreamingVideoInput> GetVideoInput() const { return VideoInput; }
+	TSharedPtr<FSceneViewport> GetViewport() const { return Viewport.Pin(); }
+
+private:
+	void SetupVideoInput();
+
+private:
+	TSharedPtr<FPixelStreamingVideoInput> VideoInput;
+	TWeakPtr<FSceneViewport> Viewport;
+};
diff --git a/Engine/Plugins/Media/PixelStreamingVCam/Source/PixelStreamingVCam/Private/PixelStreamingMediaOutput.cpp b/Engine/Plugins/Media/PixelStreamingVCam/Source/PixelStreamingVCam/Private/PixelStreamingMediaOutput.cpp
new file mode 100644
--- /dev/null
+++ b/Engine/Plugins/Media/PixelStreamingVCam/Source/PixelStreamingVCam/Private/PixelStreamingMediaOutput.cpp
@@ -0,0 +1,76 @@
+// Copyright Epic Games, Inc. All Rights Reserved.
+
+#include "PixelStreamingMediaOutput.h"
+#include "PixelStreamingMediaCapture.h"
+#include "IPixelStreamingModule.h"
+
+void UPixelStreamingMediaOutput::BeginDestroy()
+{
+	StopStreaming();
+	Streamer = nullptr;
+	Super::BeginDestroy();
+}
+
+UMediaCapture* UPixelStreamingMediaOutput::CreateMediaCaptureImpl()
+{
+	if (!Streamer.IsValid())
+	{
+		IPixelStreamingModule& Module = FModuleManager::LoadModuleChecked<IPixelStreamingModule>("PixelStreaming");
+		Streamer = Module.CreateStreamer(StreamID);
+	}
+
+	Capture = nullptr;
+	if (Streamer.IsValid())
+	{
+		Capture = NewObject<UPixelStreamingMediaCapture>();
+		Capture->SetMediaOutput(this);
+		Capture->OnStateChangedNative.AddLambda([this]() { OnCaptureStateChanged(); });
+	}
+	return Capture;
+}
+
+void UPixelStreamingMediaOutput::OnCaptureStateChanged()
+{
+	switch (Capture->GetState())
+	{
+		case EMediaCaptureState::Capturing:
+			StartStreaming();
+			break;
+		case EMediaCaptureState::Stopped:
+		case EMediaCaptureState::Error:
+			StopStreaming();
+			break;
+		default:
+			break;
+	}
+}
+
+void UPixelStreamingMediaOutput::StartStreaming()
+{
+	if (Streamer)
+	{
+		Streamer->SetVideoInput(Capture->GetVideoInput());
+		Streamer->SetTargetViewport(Capture->GetViewport());
+		Streamer->SetSignallingServerURL(SignallingServerURL);
+		Streamer->StartStreaming();
+	}
+}
+
+void UPixelStreamingMediaOutput::StopStreaming()
+{
+	if (Streamer)
+	{
+		Streamer->StopStreaming();
+		Streamer->SetTargetViewport(nullptr);
+	}
+}
+
+void UPixelStreamingMediaOutput::SetSignallingServerURL(FString InURL)
+{
+	SignallingServerURL = InURL;
+}
+
+void UPixelStreamingMediaOutput::SetSignallingStreamID(FString InStreamID)
+{
+	StreamID = InStreamID;
+}
\ No newline at end of file
diff --git a/Engine/Plugins/Media/PixelStreamingVCam/Source/PixelStreamingVCam/Private/PixelStreamingMediaOutput.h b/Engine/Plugins/Media/PixelStreamingVCam/Source/PixelStreamingVCam/Private/PixelStreamingMediaOutput.h
new file mode 100644
--- /dev/null
+++ b/Engine/Plugins/Media/PixelStreamingVCam/Source/PixelStreamingVCam/Private/PixelStreamingMediaOutput.h
@@ -0,0 +1,44 @@
+// Copyright Epic Games, Inc. All Rights Reserved.
+
+#pragma once
+
+#include "MediaOutput.h"
+#include "IPixelStreamingStreamer.h"
+#include "PixelStreamingMediaCapture.h"
+#include "PixelStreamingMediaOutput.generated.h"
+
+
+UCLASS(BlueprintType)
+class UPixelStreamingMediaOutput : public UMediaOutput
+{
+	GENERATED_BODY()
+
+	//~ Begin UMediaOutput interface
+public:
+	virtual FIntPoint GetRequestedSize() const override { return UMediaOutput::RequestCaptureSourceSize; }
+	virtual EPixelFormat GetRequestedPixelFormat() const override { return EPixelFormat::PF_B8G8R8A8; }
+	virtual EMediaCaptureConversionOperation GetConversionOperation(EMediaCaptureSourceType InSourceType) const override { return EMediaCaptureConversionOperation::CUSTOM; }
+
+	virtual void BeginDestroy() override;
+
+protected:
+	virtual UMediaCapture* CreateMediaCaptureImpl() override;
+	//~ End UMediaOutput interface
+
+public:
+	TSharedPtr<IPixelStreamingStreamer> GetStreamer() const { return Streamer; }
+
+	void SetSignallingServerURL(FString InURL);
+	void SetSignallingStreamID(FString InStreamID);
+
+private:
+	UPixelStreamingMediaCapture* Capture = nullptr;
+	TSharedPtr<IPixelStreamingStreamer> Streamer;
+
+	FString SignallingServerURL;
+	FString StreamID = TEXT("VCam");
+
+	void OnCaptureStateChanged();
+	void StartStreaming();
+	void StopStreaming();
+};
diff --git a/Engine/Plugins/Media/PixelStreamingVCam/Source/PixelStreamingVCam/Private/VCamPixelStreamingSession.cpp b/Engine/Plugins/Media/PixelStreamingVCam/Source/PixelStreamingVCam/Private/VCamPixelStreamingSession.cpp
new file mode 100644
--- /dev/null
+++ b/Engine/Plugins/Media/PixelStreamingVCam/Source/PixelStreamingVCam/Private/VCamPixelStreamingSession.cpp
@@ -0,0 +1,148 @@
+// Copyright Epic Games, Inc. All Rights Reserved.
+
+#include "VCamPixelStreamingSession.h"
+#include "VCamOutputComposure.h"
+#include "PixelStreamingServers.h"
+#include "Editor/EditorPerformanceSettings.h"
+#include "Containers/UnrealString.h"
+
+DEFINE_LOG_CATEGORY(LogVCamOutputProvider);
+namespace VCamPixelStreamingSession
+{
+	static const FName LevelEditorName(TEXT("LevelEditor"));
+	static const FSoftClassPath EmptyUMGSoftClassPath(TEXT("/VCamCore/Assets/VCam_EmptyVisibleUMG.VCam_EmptyVisibleUMG_C"));
+} // namespace VCamPixelStreamingSession
+
+void UVCamPixelStreamingSession::Initialize()
+{
+	if (!bInitialized && (MediaOutput == nullptr))
+	{
+		MediaOutput = NewObject<UPixelStreamingMediaOutput>(GetTransientPackage(), UPixelStreamingMediaOutput::StaticClass());
+	}
+	Super::Initialize();
+}
+
+void UVCamPixelStreamingSession::Deinitialize()
+{
+	if (MediaOutput)
+	{
+		MediaOutput->ConditionalBeginDestroy();
+	}
+	MediaOutput = nullptr;
+	Super::Deinitialize();
+}
+
+void UVCamPixelStreamingSession::Activate()
+{
+	// If we don't have a UMG assigned, we still need to create an empty 'dummy' UMG in order to properly route the input back from the RemoteSession device
+	if (UMGClass == nullptr)
+	{
+		bUsingDummyUMG = true;
+		UMGClass = VCamPixelStreamingSession::EmptyUMGSoftClassPath.TryLoadClass<UUserWidget>();
+	}
+
+	if (!bInitialized)
+	{
+		UE_LOG(LogVCamOutputProvider, Warning, TEXT("Trying to start Pixel Streaming, but has not been initialized yet"));
+		SetActive(false);
+		return;
+	}
+
+	UEditorPerformanceSettings* Settings = GetMutableDefault<UEditorPerformanceSettings>();
+	bOldThrottleCPUWhenNotForeground = Settings->bThrottleCPUWhenNotForeground;
+	if (PreventEditorIdle)
+	{
+		Settings->bThrottleCPUWhenNotForeground = false;
+		Settings->PostEditChange();
+	}
+
+	if (StartSignallingServer)
+	{
+		CirrusProcess = UE::PixelStreaming::Servers::MakeSignallingServer();
+		UE::PixelStreaming::Servers::FLaunchArgs LaunchArgs;
+		LaunchArgs.bEphemeral = false;
+		LaunchArgs.ProcessArgs = FString::Printf(TEXT("--StreamerPort=%s"), *FString::FromInt(PortNumber));
+		CirrusProcess->Launch(LaunchArgs);
+	}
+
+	MediaOutput->SetSignallingServerURL(FString::Printf(TEXT("%s:%s"), *IP, *FString::FromInt(PortNumber)));
+	UE_LOG(LogVCamOutputProvider, Log, TEXT("Activating PixelStreaming VCam Session. Endpoint: %s:%s"), *IP, *FString::FromInt(PortNumber));
+
+	MediaCapture = Cast<UPixelStreamingMediaCapture>(MediaOutput->CreateMediaCapture());
+
+	FMediaCaptureOptions Options;
+	Options.bResizeSourceBuffer = true;
+
+	// If we are rendering from a ComposureOutputProvider, get the requested render target and use that instead of the viewport
+	if (UVCamOutputComposure* ComposureProvider = Cast<UVCamOutputComposure>(GetOtherOutputProviderByIndex(FromComposureOutputProviderIndex)))
+	{
+		if (ComposureProvider->FinalOutputRenderTarget)
+		{
+			MediaCapture->CaptureTextureRenderTarget2D(ComposureProvider->FinalOutputRenderTarget, Options);
+			UE_LOG(LogVCamOutputProvider, Log, TEXT("PixelStreaming set with ComposureRenderTarget"));
+		}
+		else
+		{
+			UE_LOG(LogVCamOutputProvider, Warning, TEXT("PixelStreaming Composure usage was requested, but the specified ComposureOutputProvider has no FinalOutputRenderTarget set"));
+		}
+	}
+	else
+	{
+		TWeakPtr<FSceneViewport> SceneViewport = GetTargetSceneViewport();
+		if (TSharedPtr<FSceneViewport> PinnedSceneViewport = SceneViewport.Pin())
+		{
+			MediaCapture->CaptureSceneViewport(PinnedSceneViewport, Options);
+			UE_LOG(LogVCamOutputProvider, Log, TEXT("PixelStreaming set with viewport"));
+		}
+	}
+	Super::Activate();
+}
+
+void UVCamPixelStreamingSession::Deactivate()
+{
+	if (MediaCapture)
+	{
+		MediaCapture->StopCapture(true);
+		MediaCapture = nullptr;
+	}
+	Super::Deactivate();
+	if (bUsingDummyUMG)
+	{
+		UMGClass = nullptr;
+		bUsingDummyUMG = false;
+	}
+
+	if(CirrusProcess.IsValid())
+	{
+		CirrusProcess->Stop();
+	}
+	
+	UEditorPerformanceSettings* Settings = GetMutableDefault<UEditorPerformanceSettings>();
+	Settings->bThrottleCPUWhenNotForeground = bOldThrottleCPUWhenNotForeground;
+	Settings->PostEditChange();
+}
+
+void UVCamPixelStreamingSession::Tick(const float DeltaTime)
+{
+	Super::Tick(DeltaTime);
+}
+
+#if WITH_EDITOR
+void UVCamPixelStreamingSession::PostEditChangeProperty(FPropertyChangedEvent& PropertyChangedEvent)
+{
+	FProperty* Property = PropertyChangedEvent.MemberProperty;
+	if (Property && PropertyChangedEvent.ChangeType != EPropertyChangeType::Interactive)
+	{
+		static FName NAME_IP = GET_MEMBER_NAME_CHECKED(UVCamPixelStreamingSession, IP);
+		static FName NAME_PortNumber = GET_MEMBER_NAME_CHECKED(UVCamPixelStreamingSession, PortNumber);
+		static FName NAME_FromComposureOutputProviderIndex = GET_MEMBER_NAME_CHECKED(UVCamPixelStreamingSession, FromComposureOutputProviderIndex);
+		if ((Property->GetFName() == NAME_IP) || (Property->GetFName() == NAME_PortNumber) || (Property->GetFName() == NAME_FromComposureOutputProviderIndex))
+		{
+			SetActive(!bIsActive);
+		}
+	}
+	Super::PostEditChangeProperty(PropertyChangedEvent);
+}
+#endif
+
+IMPLEMENT_MODULE(FDefaultModuleImpl, PixelStreamingVCam)
diff --git a/Engine/Plugins/Media/PixelStreamingVCam/Source/PixelStreamingVCam/Public/VCamPixelStreamingSession.h b/Engine/Plugins/Media/PixelStreamingVCam/Source/PixelStreamingVCam/Public/VCamPixelStreamingSession.h
new file mode 100644
--- /dev/null
+++ b/Engine/Plugins/Media/PixelStreamingVCam/Source/PixelStreamingVCam/Public/VCamPixelStreamingSession.h
@@ -0,0 +1,60 @@
+// Copyright Epic Games, Inc. All Rights Reserved.
+
+#pragma once
+
+#include "VCamOutputProviderBase.h"
+#include "PixelStreamingMediaCapture.h"
+#include "PixelStreamingMediaOutput.h"
+#include "PixelStreamingServers.h"
+#include "Slate/SceneViewport.h"
+#include "Engine/TextureRenderTarget2D.h"
+#if WITH_EDITOR
+#include "LevelEditorViewport.h"
+#endif
+#include "VCamPixelStreamingSession.generated.h"
+
+UCLASS(meta = (DisplayName = "Pixel Streaming Provider"))
+class PIXELSTREAMINGVCAM_API UVCamPixelStreamingSession : public UVCamOutputProviderBase
+{
+	GENERATED_BODY()
+public:
+	virtual void Initialize() override;
+	virtual void Deinitialize() override;
+	virtual void Activate() override;
+	virtual void Deactivate() override;
+	virtual void Tick(const float DeltaTime) override;
+#if WITH_EDITOR
+	virtual void PostEditChangeProperty(FPropertyChangedEvent& PropertyChangedEvent) override;
+#endif
+	UPROPERTY(EditAnywhere, BlueprintReadWrite, Category = "Output")
+	FString IP = TEXT("ws://127.0.0.1");
+
+	// Network port number - change this only if connecting multiple RemoteSession devices to the same PC
+	UPROPERTY(EditAnywhere, BlueprintReadWrite, Category = "Output")
+	int32 PortNumber = 8888;
+
+	// If using the output from a Composure Output Provider, specify it here
+	UPROPERTY(EditAnywhere, BlueprintReadWrite, Category = "Output")
+	int32 FromComposureOutputProviderIndex = INDEX_NONE;
+
+	// Check this if you are not separately running a signalling server for this session
+	UPROPERTY(EditAnywhere, BlueprintReadWrite, Category = "Output")
+	bool StartSignallingServer = true;
+
+	// If not selected, when the editor is not the foreground application, input through the vcam session may seem sluggish or unresponsive.
+	UPROPERTY(EditAnywhere, BlueprintReadWrite, Category = "Output")
+	bool PreventEditorIdle = true;
+
+protected:
+	UPROPERTY(Transient)
+	UPixelStreamingMediaOutput* MediaOutput = nullptr;
+
+	UPROPERTY(Transient)
+	UPixelStreamingMediaCapture* MediaCapture = nullptr;
+
+private:
+	FHitResult LastViewportTouchResult;
+	bool bUsingDummyUMG = false;
+	bool bOldThrottleCPUWhenNotForeground;
+	TSharedPtr<UE::PixelStreaming::Servers::FServerBase> CirrusProcess;
+};
