Update PixelStreamingVCam plugin definition to not use unacceptable words

diff --git a/Engine/Plugins/Media/PixelStreamingVCam/PixelStreamingVCam.uplugin b/Engine/Plugins/Media/PixelStreamingVCam/PixelStreamingVCam.uplugin
--- a/Engine/Plugins/Media/PixelStreamingVCam/PixelStreamingVCam.uplugin
+++ b/Engine/Plugins/Media/PixelStreamingVCam/PixelStreamingVCam.uplugin
@@ -19,8 +19,8 @@
 			"Name": "PixelStreamingVCam",
 			"Type": "Editor",
 			"LoadingPhase": "PostEngineInit",
-			"WhitelistPlatforms": [ "Win64", "Linux" ],
-			"BlacklistTargets": [ "Server" ]
+			"PlatformAllowList": [ "Win64", "Linux" ],
+			"TargetDenyList": [ "Server" ]
 		}
 	],
 	"Plugins": [
