Fixing VCam to work with the updated MediaCapture interface.

diff --git a/Engine/Plugins/Media/PixelStreamingVCam/Source/PixelStreamingVCam/Private/PixelStreamingMediaCapture.cpp b/Engine/Plugins/Media/PixelStreamingVCam/Source/PixelStreamingVCam/Private/PixelStreamingMediaCapture.cpp
--- a/Engine/Plugins/Media/PixelStreamingVCam/Source/PixelStreamingVCam/Private/PixelStreamingMediaCapture.cpp
+++ b/Engine/Plugins/Media/PixelStreamingVCam/Source/PixelStreamingVCam/Private/PixelStreamingMediaCapture.cpp
@@ -4,19 +4,14 @@
 #include "PixelStreamingSourceFrame.h"
 #include "Slate/SceneViewport.h"
 
-void UPixelStreamingMediaCapture::OnCustomCapture_RenderingThread(
-	FRHICommandListImmediate& RHICmdList,
-	const FCaptureBaseData& InBaseData,
-	TSharedPtr<FMediaCaptureUserData, ESPMode::ThreadSafe> InUserData,
-	FTexture2DRHIRef InSourceTexture,
-	FTextureRHIRef TargetableTexture,
-	FResolveParams& ResolveParams,
-	FVector2D CropU,
-	FVector2D CropV)
+void UPixelStreamingMediaCapture::OnRHIResourceCaptured_RenderingThread(
+		const FCaptureBaseData& InBaseData,
+		TSharedPtr<FMediaCaptureUserData, ESPMode::ThreadSafe> InUserData,
+		FTextureRHIRef InTexture)
 {
 	if (VideoInput)
 	{
-		VideoInput->OnFrame.Broadcast(FPixelStreamingSourceFrame(InSourceTexture));
+		VideoInput->OnFrame.Broadcast(FPixelStreamingSourceFrame(InTexture));
 	}
 }
 
diff --git a/Engine/Plugins/Media/PixelStreamingVCam/Source/PixelStreamingVCam/Private/PixelStreamingMediaCapture.h b/Engine/Plugins/Media/PixelStreamingVCam/Source/PixelStreamingVCam/Private/PixelStreamingMediaCapture.h
--- a/Engine/Plugins/Media/PixelStreamingVCam/Source/PixelStreamingVCam/Private/PixelStreamingMediaCapture.h
+++ b/Engine/Plugins/Media/PixelStreamingVCam/Source/PixelStreamingVCam/Private/PixelStreamingMediaCapture.h
@@ -13,18 +13,14 @@ class UPixelStreamingMediaCapture : public UMediaCapture
 
 	//~ Begin UMediaCapture interface
 public:
-	virtual void OnCustomCapture_RenderingThread(
-		FRHICommandListImmediate& RHICmdList,
+	virtual void OnRHIResourceCaptured_RenderingThread(
 		const FCaptureBaseData& InBaseData,
 		TSharedPtr<FMediaCaptureUserData, ESPMode::ThreadSafe> InUserData,
-		FTexture2DRHIRef InSourceTexture,
-		FTextureRHIRef TargetableTexture,
-		FResolveParams& ResolveParams,
-		FVector2D CropU,
-		FVector2D CropV) override;
+		FTextureRHIRef InTexture) override;
+
 	virtual bool CaptureSceneViewportImpl(TSharedPtr<FSceneViewport>& InSceneViewport) override;
 	virtual bool CaptureRenderTargetImpl(UTextureRenderTarget2D* InRenderTarget) override;
-	virtual bool ShouldCaptureRHITexture() const { return true; }
+	virtual bool ShouldCaptureRHIResource() const { return true; }
 	virtual void StopCaptureImpl(bool bAllowPendingFrameToBeProcess) override {}
 	//~ End UMediaCapture interface
 
diff --git a/Engine/Plugins/Media/PixelStreamingVCam/Source/PixelStreamingVCam/Private/PixelStreamingMediaOutput.h b/Engine/Plugins/Media/PixelStreamingVCam/Source/PixelStreamingVCam/Private/PixelStreamingMediaOutput.h
--- a/Engine/Plugins/Media/PixelStreamingVCam/Source/PixelStreamingVCam/Private/PixelStreamingMediaOutput.h
+++ b/Engine/Plugins/Media/PixelStreamingVCam/Source/PixelStreamingVCam/Private/PixelStreamingMediaOutput.h
@@ -17,7 +17,7 @@ class UPixelStreamingMediaOutput : public UMediaOutput
 public:
 	virtual FIntPoint GetRequestedSize() const override { return UMediaOutput::RequestCaptureSourceSize; }
 	virtual EPixelFormat GetRequestedPixelFormat() const override { return EPixelFormat::PF_B8G8R8A8; }
-	virtual EMediaCaptureConversionOperation GetConversionOperation(EMediaCaptureSourceType InSourceType) const override { return EMediaCaptureConversionOperation::CUSTOM; }
+	virtual EMediaCaptureConversionOperation GetConversionOperation(EMediaCaptureSourceType InSourceType) const override { return EMediaCaptureConversionOperation::NONE; }
 
 	virtual void BeginDestroy() override;
 
