Removed ARKit coordinate space conversions now that the ipad app is already sending in the correct format.

Added some additional safeguards to reading the data to warn about incorrect data sizes

diff --git a/Engine/Plugins/Media/PixelStreamingVCam/Source/PixelStreamingVCam/Private/VCamPixelStreamingSession.cpp b/Engine/Plugins/Media/PixelStreamingVCam/Source/PixelStreamingVCam/Private/VCamPixelStreamingSession.cpp
--- a/Engine/Plugins/Media/PixelStreamingVCam/Source/PixelStreamingVCam/Private/VCamPixelStreamingSession.cpp
+++ b/Engine/Plugins/Media/PixelStreamingVCam/Source/PixelStreamingVCam/Private/VCamPixelStreamingSession.cpp
@@ -117,49 +117,43 @@ void UVCamPixelStreamingSession::Activate()
 			{
 				return;
 			}
+
+			// 1 byte for type, 16 floats for transform, 1 double for timestamp
+			static constexpr int ExpectedDataSize = 1 + (16 * 4) + 8;
+			if (Data.Num() != ExpectedDataSize)
+			{
+				UE_LOG(LogVCamOutputProvider, Warning, TEXT("ARKitTransform data received but with invalid size. Expects %d bytes but received %d"),
+					ExpectedDataSize,
+					Data.Num());
+
+				return;
+			}
 			
 			FMemoryReader Buffer(Data);
+			
 			uint8 MessageType;
 			Buffer << MessageType;
-			
-			TArray<FPlane> Columns;
-			for(int32 Idx = 0; Idx < 4; Idx++)
+
+			// The buffer contains the transform matrix stored as 16 floats
+			FMatrix ARKitMatrix;
+			for (int32 Row = 0; Row < 4; ++Row)
 			{
-				float x;
-				float y;
-				float z;
-				float w;
-
-				Buffer << x;
-				Buffer << y;
-				Buffer << z;
-				Buffer << w;
-
-				FVector4 Column(x, y, z, w);
-				Columns.Add(FPlane(Column));
+				float Col0, Col1, Col2, Col3;
+				Buffer << Col0 << Col1 << Col2 << Col3;
+				ARKitMatrix.M[Row][0] = Col0;
+				ARKitMatrix.M[Row][1] = Col1;
+				ARKitMatrix.M[Row][2] = Col2;
+				ARKitMatrix.M[Row][3] = Col3;
 			}
+			ARKitMatrix.DiagnosticCheckNaN();
 
-			FMatrix RawYUpFMatrix(Columns[0], Columns[1], Columns[2], Columns[3]);
-			/**
-			 * Convert's an ARKit 'Y up' 'right handed' coordinate system transform to Unreal's 'Z up' 
-			 * 'left handed' coordinate system.
-			 * Taken from 'Engine\Plugins\Runtime\AR\AppleAR\AppleARKit\Source\AppleARKit\Public\AppleARKitConversion.h ToFTransform'
-			 *
-			 * Ignores scale.
-			 */
-			// Extract & convert translation
-			FVector Translation = FVector( -RawYUpFMatrix.M[3][2], RawYUpFMatrix.M[3][0], RawYUpFMatrix.M[3][1] ) * 100.0f;
-			// Extract & convert rotation 
-			FQuat RawRotation( RawYUpFMatrix );
-			FQuat Rotation( -RawRotation.Z, RawRotation.X, RawRotation.Y, -RawRotation.W );
-			FRotator ModifiedRotation(Rotation);
-			ModifiedRotation.Roll -= 90.0f;
-
-			TSharedPtr<FPixelStreamingLiveLinkSource> LiveLinkSource = UVCamPixelStreamingSubsystem::Get()->LiveLinkSource;
-			if(LiveLinkSource)
+			// 
+			double Timestamp;
+			Buffer << Timestamp;
+
+			if(TSharedPtr<FPixelStreamingLiveLinkSource> LiveLinkSource = UVCamPixelStreamingSubsystem::Get()->LiveLinkSource)
 			{
-				FTransform Transform(ModifiedRotation, Translation, FVector(RawYUpFMatrix.GetScaleVector(1.0f)));
-				LiveLinkSource->PushTransformForSubject(GetFName(), Transform);
+				LiveLinkSource->PushTransformForSubject(GetFName(), FTransform(ARKitMatrix));
 			}
 
 			/**
