Modifed PixelStreamingLiveLinkSource to support Timecode and defaulted the source to evaluate at latest frame for lower latency

diff --git a/Engine/Plugins/Media/PixelStreamingVCam/Source/PixelStreamingVCam/Private/VCamPixelStreamingLiveLink.cpp b/Engine/Plugins/Media/PixelStreamingVCam/Source/PixelStreamingVCam/Private/VCamPixelStreamingLiveLink.cpp
--- a/Engine/Plugins/Media/PixelStreamingVCam/Source/PixelStreamingVCam/Private/VCamPixelStreamingLiveLink.cpp
+++ b/Engine/Plugins/Media/PixelStreamingVCam/Source/PixelStreamingVCam/Private/VCamPixelStreamingLiveLink.cpp
@@ -7,6 +7,13 @@
 
 #define LOCTEXT_NAMESPACE "PixelStreamingLiveLinkSource"
 
+UPixelStreamingLiveLinkSourceSettings::UPixelStreamingLiveLinkSourceSettings()
+	: ULiveLinkSourceSettings()
+{
+	// Override the default evaluation mode to latest
+	Mode = ELiveLinkSourceMode::Latest;
+}
+
 FPixelStreamingLiveLinkSource::FPixelStreamingLiveLinkSource()
 	: LiveLinkClient(nullptr)
 {
@@ -84,4 +91,23 @@ void FPixelStreamingLiveLinkSource::PushTransformForSubject(FName SubjectName, F
 	}
 }
 
+void FPixelStreamingLiveLinkSource::PushTransformForSubject(FName SubjectName, FTransform Transform,
+	double Timestamp) const
+{
+	if (LiveLinkClient)
+	{
+		const FLiveLinkSubjectKey SubjectKey(SourceGuid, SubjectName);
+		FLiveLinkFrameDataStruct FrameDataStruct(FLiveLinkTransformFrameData::StaticStruct());
+		FLiveLinkTransformFrameData* TransformFrameData = FrameDataStruct.Cast<FLiveLinkTransformFrameData>();
+		TransformFrameData->Transform = Transform;
+
+		// Timestamp is currently assumed to be elapsed seconds at a fixed rate of 60 frames per second
+		// this will be adjusted as actual rate information is supported
+		const int32 NumberOfFrames = FMath::FloorToInt32(Timestamp * 60.0);
+		TransformFrameData->MetaData.SceneTime = FQualifiedFrameTime(NumberOfFrames, FFrameRate(60,1));
+		
+		LiveLinkClient->PushSubjectFrameData_AnyThread(SubjectKey, MoveTemp(FrameDataStruct));
+	}
+}
+
 #undef LOCTEXT_NAMESPACE
diff --git a/Engine/Plugins/Media/PixelStreamingVCam/Source/PixelStreamingVCam/Private/VCamPixelStreamingLiveLink.h b/Engine/Plugins/Media/PixelStreamingVCam/Source/PixelStreamingVCam/Private/VCamPixelStreamingLiveLink.h
--- a/Engine/Plugins/Media/PixelStreamingVCam/Source/PixelStreamingVCam/Private/VCamPixelStreamingLiveLink.h
+++ b/Engine/Plugins/Media/PixelStreamingVCam/Source/PixelStreamingVCam/Private/VCamPixelStreamingLiveLink.h
@@ -5,6 +5,17 @@
 #include "CoreMinimal.h"
 #include "ILiveLinkSource.h"
 
+#include "VCamPixelStreamingLiveLink.generated.h"
+
+UCLASS()
+class UPixelStreamingLiveLinkSourceSettings : public ULiveLinkSourceSettings
+{
+public:
+	GENERATED_BODY()
+
+	UPixelStreamingLiveLinkSourceSettings();
+};
+
 class FPixelStreamingLiveLinkSource : public ILiveLinkSource
 {
 public:
@@ -18,11 +29,16 @@ public:
 	virtual FText GetSourceType() const override;
 	virtual FText GetSourceMachineName() const override;
 	virtual FText GetSourceStatus() const override;
+
+	// Override the settings class to allow us to customize the default values
+	virtual TSubclassOf<ULiveLinkSourceSettings> GetSettingsClass() const override{ return UPixelStreamingLiveLinkSourceSettings::StaticClass(); };
+	
 	// Registers a new subject with the Transform Role to the Live Link Client
 	// If called with a subject name that already exists in this source then this will reset any buffered data for that subject
 	void CreateSubject(FName SubjectName) const;
 	void RemoveSubject(FName SubjectName) const;
 	void PushTransformForSubject(FName SubjectName, FTransform Transform) const;
+	void PushTransformForSubject(FName SubjectName, FTransform Transform, double Timestamp) const;
 
 private:
 	// Cached information for communicating with the live link client
diff --git a/Engine/Plugins/Media/PixelStreamingVCam/Source/PixelStreamingVCam/Private/VCamPixelStreamingSession.cpp b/Engine/Plugins/Media/PixelStreamingVCam/Source/PixelStreamingVCam/Private/VCamPixelStreamingSession.cpp
--- a/Engine/Plugins/Media/PixelStreamingVCam/Source/PixelStreamingVCam/Private/VCamPixelStreamingSession.cpp
+++ b/Engine/Plugins/Media/PixelStreamingVCam/Source/PixelStreamingVCam/Private/VCamPixelStreamingSession.cpp
@@ -153,7 +153,7 @@ void UVCamPixelStreamingSession::Activate()
 
 			if(TSharedPtr<FPixelStreamingLiveLinkSource> LiveLinkSource = UVCamPixelStreamingSubsystem::Get()->LiveLinkSource)
 			{
-				LiveLinkSource->PushTransformForSubject(GetFName(), FTransform(ARKitMatrix));
+				LiveLinkSource->PushTransformForSubject(GetFName(), FTransform(ARKitMatrix), Timestamp);
 			}
 
 			/**
