Fix PixelStreamingVCam plugin unable to connect to auto-started signalling sever.

- Make iPad app always send websocket message using string as UE's websocket implementation is non-standard for binary messages.
- Clean up auto-start signalling server in UE to stop signalling if start signalling checkbox is modified.
- Clean up remnants of old method of running signalling server.
- Move OnARKitTransformReceived into its own function for code cleanliness.

diff --git a/Engine/Extras/VirtualProduction/LiveLinkVCAM/vcam/WebRTC/Services/WebSocketProvider/NativeWebSocketProvider.swift b/Engine/Extras/VirtualProduction/LiveLinkVCAM/vcam/WebRTC/Services/WebSocketProvider/NativeWebSocketProvider.swift
--- a/Engine/Extras/VirtualProduction/LiveLinkVCAM/vcam/WebRTC/Services/WebSocketProvider/NativeWebSocketProvider.swift
+++ b/Engine/Extras/VirtualProduction/LiveLinkVCAM/vcam/WebRTC/Services/WebSocketProvider/NativeWebSocketProvider.swift
@@ -39,8 +39,12 @@ class NativeWebSocketProvider : NSObject, WebSocketProvider, URLSessionWebSocket
     }
     
     func send(data: Data) {
-        socket.send(.data(data)) { error in
-            self.handleError(error)
+        let stringifiedMsg = String(data: data, encoding: .utf8)
+        if let msg = stringifiedMsg {
+            self.send(msg: msg)
+        }
+        else {
+            debugPrint("Could not convert data message to UTF8 string.")
         }
     }
     
diff --git a/Engine/Extras/VirtualProduction/LiveLinkVCAM/vcam/WebRTC/Services/WebSocketProvider/StarscreamProvider.swift b/Engine/Extras/VirtualProduction/LiveLinkVCAM/vcam/WebRTC/Services/WebSocketProvider/StarscreamProvider.swift
--- a/Engine/Extras/VirtualProduction/LiveLinkVCAM/vcam/WebRTC/Services/WebSocketProvider/StarscreamProvider.swift
+++ b/Engine/Extras/VirtualProduction/LiveLinkVCAM/vcam/WebRTC/Services/WebSocketProvider/StarscreamProvider.swift
@@ -24,7 +24,17 @@ class StarscreamWebSocket: WebSocketProvider {
     }
     
     func send(data: Data) {
-        self.socket.write(data: data)
+        let stringifiedMsg = String(data: data, encoding: .utf8)
+        if let msg = stringifiedMsg {
+            self.send(msg: msg)
+        }
+        else {
+            debugPrint("Could not convert data message to UTF8 string.")
+        }
+    }
+
+    func send(msg: String) {
+        self.socket.write(string: msg)
     }
     
     func close() {
diff --git a/Engine/Extras/VirtualProduction/LiveLinkVCAM/vcam/WebRTC/Services/WebSocketProvider/WebSocketProvider.swift b/Engine/Extras/VirtualProduction/LiveLinkVCAM/vcam/WebRTC/Services/WebSocketProvider/WebSocketProvider.swift
--- a/Engine/Extras/VirtualProduction/LiveLinkVCAM/vcam/WebRTC/Services/WebSocketProvider/WebSocketProvider.swift
+++ b/Engine/Extras/VirtualProduction/LiveLinkVCAM/vcam/WebRTC/Services/WebSocketProvider/WebSocketProvider.swift
@@ -13,6 +13,7 @@ protocol WebSocketProvider: AnyObject {
     func connect()
     func close()
     func send(data: Data)
+    func send(msg: String)
 }
 
 protocol WebSocketProviderDelegate: AnyObject {
diff --git a/Engine/Plugins/Media/PixelStreamingVCam/Source/PixelStreamingVCam/Private/VCamPixelStreamingSession.cpp b/Engine/Plugins/Media/PixelStreamingVCam/Source/PixelStreamingVCam/Private/VCamPixelStreamingSession.cpp
--- a/Engine/Plugins/Media/PixelStreamingVCam/Source/PixelStreamingVCam/Private/VCamPixelStreamingSession.cpp
+++ b/Engine/Plugins/Media/PixelStreamingVCam/Source/PixelStreamingVCam/Private/VCamPixelStreamingSession.cpp
@@ -111,65 +111,69 @@ void UVCamPixelStreamingSession::Activate()
 
 	if(MediaOutput->GetStreamer())
 	{
-		MediaOutput->GetStreamer()->OnInputReceived.AddLambda([this](FPixelStreamingPlayerId PlayerId, uint8 Type, TArray<uint8> Data)
-		{
-			if(Type != static_cast<uint8>(UE::PixelStreaming::Protocol::EToStreamerMsg::ARKitTransform) || !EnableARKitTracking)
-			{
-				return;
-			}
+		MediaOutput->GetStreamer()->OnInputReceived.AddUObject(this, &UVCamPixelStreamingSession::OnARKitTransformReceived);
+	}
 
-			// 1 byte for type, 16 floats for transform, 1 double for timestamp
-			static constexpr int ExpectedDataSize = 1 + (16 * 4) + 8;
-			if (Data.Num() != ExpectedDataSize)
-			{
-				UE_LOG(LogPixelStreamingVCam, Warning, TEXT("ARKitTransform data received but with invalid size. Expects %d bytes but received %d"),
-					ExpectedDataSize,
-					Data.Num());
+	Super::Activate();
+}
 
-				return;
-			}
-			
-			FMemoryReader Buffer(Data);
-			
-			uint8 MessageType;
-			Buffer << MessageType;
-
-			// The buffer contains the transform matrix stored as 16 floats
-			FMatrix ARKitMatrix;
-			for (int32 Row = 0; Row < 4; ++Row)
-			{
-				float Col0, Col1, Col2, Col3;
-				Buffer << Col0 << Col1 << Col2 << Col3;
-				ARKitMatrix.M[Row][0] = Col0;
-				ARKitMatrix.M[Row][1] = Col1;
-				ARKitMatrix.M[Row][2] = Col2;
-				ARKitMatrix.M[Row][3] = Col3;
-			}
-			ARKitMatrix.DiagnosticCheckNaN();
+void UVCamPixelStreamingSession::OnARKitTransformReceived(FPixelStreamingPlayerId PlayerId, uint8 Type, TArray<uint8> Data)
+{
+	if(Type != static_cast<uint8>(UE::PixelStreaming::Protocol::EToStreamerMsg::ARKitTransform) || !EnableARKitTracking)
+	{
+		return;
+	}
 
-			// 
-			double Timestamp;
-			Buffer << Timestamp;
+	// 1 byte for type, 16 floats for transform, 1 double for timestamp
+	static constexpr int ExpectedDataSize = 1 + (16 * 4) + 8;
+	if (Data.Num() != ExpectedDataSize)
+	{
+		UE_LOG(LogPixelStreamingVCam, Warning, TEXT("ARKitTransform data received but with invalid size. Expects %d bytes but received %d"),
+			ExpectedDataSize,
+			Data.Num());
 
-			if(TSharedPtr<FPixelStreamingLiveLinkSource> LiveLinkSource = UVCamPixelStreamingSubsystem::Get()->LiveLinkSource)
-			{
-				LiveLinkSource->PushTransformForSubject(GetFName(), FTransform(ARKitMatrix), Timestamp);
-			}
+		return;
+	}
+	
+	FMemoryReader Buffer(Data);
+	
+	// Extract message type
+	uint8 MessageType;
+	Buffer << MessageType;
+
+	// Extract transform, the buffer contains the transform matrix stored as 16 floats
+	FMatrix ARKitMatrix;
+	for (int32 Row = 0; Row < 4; ++Row)
+	{
+		float Col0, Col1, Col2, Col3;
+		Buffer << Col0 << Col1 << Col2 << Col3;
+		ARKitMatrix.M[Row][0] = Col0;
+		ARKitMatrix.M[Row][1] = Col1;
+		ARKitMatrix.M[Row][2] = Col2;
+		ARKitMatrix.M[Row][3] = Col3;
+	}
+	ARKitMatrix.DiagnosticCheckNaN();
+
+	// Extract timestamp
+	double Timestamp;
+	Buffer << Timestamp;
 
-			/**
-			 * Code to control the level editor viewport
-			 */
-			// for (FLevelEditorViewportClient* LevelVC : GEditor->GetLevelViewportClients())
-			// {
-			// 	if (LevelVC && LevelVC->IsPerspective())
-			// 	{
-			// 		LevelVC->SetViewLocation(Translation);
-			// 		LevelVC->SetViewRotation(ModifiedRotation);
-			// 	}
-			// }
-		});
+	if(TSharedPtr<FPixelStreamingLiveLinkSource> LiveLinkSource = UVCamPixelStreamingSubsystem::Get()->LiveLinkSource)
+	{
+		LiveLinkSource->PushTransformForSubject(GetFName(), FTransform(ARKitMatrix), Timestamp);
 	}
-	Super::Activate();
+
+	/**
+	 * Code to control the level editor viewport
+	 */
+	// for (FLevelEditorViewportClient* LevelVC : GEditor->GetLevelViewportClients())
+	// {
+	// 	if (LevelVC && LevelVC->IsPerspective())
+	// 	{
+	// 		LevelVC->SetViewLocation(Translation);
+	// 		LevelVC->SetViewRotation(ModifiedRotation);
+	// 	}
+	// }
 }
 
 void UVCamPixelStreamingSession::StopSignallingServer()
@@ -232,8 +236,13 @@ void UVCamPixelStreamingSession::PostEditChangeProperty(FPropertyChangedEvent& P
 		static FName NAME_IP = GET_MEMBER_NAME_CHECKED(UVCamPixelStreamingSession, IP);
 		static FName NAME_PortNumber = GET_MEMBER_NAME_CHECKED(UVCamPixelStreamingSession, PortNumber);
 		static FName NAME_FromComposureOutputProviderIndex = GET_MEMBER_NAME_CHECKED(UVCamPixelStreamingSession, FromComposureOutputProviderIndex);
+		static FName NAME_StartSignallingServer = GET_MEMBER_NAME_CHECKED(UVCamPixelStreamingSession, StartSignallingServer);
 		
-		if ((Property->GetFName() == NAME_IP) || (Property->GetFName() == NAME_PortNumber) || (Property->GetFName() == NAME_FromComposureOutputProviderIndex))
+		// If any of these properties change then deactivate streaming so user is forced to restart streaming
+		if ((Property->GetFName() == NAME_IP) || 
+			(Property->GetFName() == NAME_PortNumber) || 
+			(Property->GetFName() == NAME_FromComposureOutputProviderIndex) || 
+			(Property->GetFName() == NAME_StartSignallingServer))
 		{
 			SetActive(false);
 		}
diff --git a/Engine/Plugins/Media/PixelStreamingVCam/Source/PixelStreamingVCam/Public/VCamPixelStreamingSession.h b/Engine/Plugins/Media/PixelStreamingVCam/Source/PixelStreamingVCam/Public/VCamPixelStreamingSession.h
--- a/Engine/Plugins/Media/PixelStreamingVCam/Source/PixelStreamingVCam/Public/VCamPixelStreamingSession.h
+++ b/Engine/Plugins/Media/PixelStreamingVCam/Source/PixelStreamingVCam/Public/VCamPixelStreamingSession.h
@@ -62,10 +62,10 @@ protected:
 
 private:
 	void StopSignallingServer();
+	void OnARKitTransformReceived(FPixelStreamingPlayerId PlayerId, uint8 Type, TArray<uint8> Data);
 
 private:
 	FHitResult LastViewportTouchResult;
 	bool bUsingDummyUMG = false;
 	bool bOldThrottleCPUWhenNotForeground;
-	TSharedPtr<UE::PixelStreamingServers::IServer> CirrusProcess;
 };
