[LiveLinkVCAM] Fix multiple issues to restore functionality to LiveLinkVCam + Pixel Streaming VCam.

- Fix syntax error in LiveLinkVCam app to restore compilation.
- Fix viewport target being set at wrong time in PixelStreamingVCam plugin.
- Remove hack from PixelStreamingVCam that is now unecessary with the move to a subsystem design.

diff --git a/Engine/Extras/VirtualProduction/LiveLinkVCAM/vcam/WebRTC/Services/WebSocketProvider/NativeWebSocketProvider.swift b/Engine/Extras/VirtualProduction/LiveLinkVCAM/vcam/WebRTC/Services/WebSocketProvider/NativeWebSocketProvider.swift
--- a/Engine/Extras/VirtualProduction/LiveLinkVCAM/vcam/WebRTC/Services/WebSocketProvider/NativeWebSocketProvider.swift
+++ b/Engine/Extras/VirtualProduction/LiveLinkVCAM/vcam/WebRTC/Services/WebSocketProvider/NativeWebSocketProvider.swift
@@ -48,8 +48,8 @@ class NativeWebSocketProvider : NSObject, WebSocketProvider, URLSessionWebSocket
         }
     }
     
-    func send(text: String) {
-        socket.send(.string(text)) { error in
+    func send(msg: String) {
+        socket.send(.string(msg)) { error in
             self.handleError(error)
         }
     }
diff --git a/Engine/Plugins/Media/PixelStreamingVCam/Source/PixelStreamingVCam/Private/PixelStreamingMediaCapture.cpp b/Engine/Plugins/Media/PixelStreamingVCam/Source/PixelStreamingVCam/Private/PixelStreamingMediaCapture.cpp
--- a/Engine/Plugins/Media/PixelStreamingVCam/Source/PixelStreamingVCam/Private/PixelStreamingMediaCapture.cpp
+++ b/Engine/Plugins/Media/PixelStreamingVCam/Source/PixelStreamingVCam/Private/PixelStreamingMediaCapture.cpp
@@ -26,6 +26,7 @@ bool UPixelStreamingMediaCapture::InitializeCapture()
 bool UPixelStreamingMediaCapture::PostInitializeCaptureViewport(TSharedPtr<FSceneViewport>& InSceneViewport)
 {
 	Viewport = InSceneViewport;
+	OnCaptureViewportInitialized.Broadcast();
 	return true;
 }
 
diff --git a/Engine/Plugins/Media/PixelStreamingVCam/Source/PixelStreamingVCam/Private/PixelStreamingMediaCapture.h b/Engine/Plugins/Media/PixelStreamingVCam/Source/PixelStreamingVCam/Private/PixelStreamingMediaCapture.h
--- a/Engine/Plugins/Media/PixelStreamingVCam/Source/PixelStreamingVCam/Private/PixelStreamingMediaCapture.h
+++ b/Engine/Plugins/Media/PixelStreamingVCam/Source/PixelStreamingVCam/Private/PixelStreamingMediaCapture.h
@@ -27,6 +27,8 @@ public:
 	TSharedPtr<IPixelStreamingVideoInput> GetVideoInput() const { return VideoInput; }
 	TSharedPtr<FSceneViewport> GetViewport() const { return Viewport.Pin(); }
 
+	DECLARE_MULTICAST_DELEGATE(FOnCaptureViewportInitialized);
+	FOnCaptureViewportInitialized OnCaptureViewportInitialized;
 private:
 	void SetupVideoInput();
 
diff --git a/Engine/Plugins/Media/PixelStreamingVCam/Source/PixelStreamingVCam/Private/PixelStreamingMediaOutput.cpp b/Engine/Plugins/Media/PixelStreamingVCam/Source/PixelStreamingVCam/Private/PixelStreamingMediaOutput.cpp
--- a/Engine/Plugins/Media/PixelStreamingVCam/Source/PixelStreamingVCam/Private/PixelStreamingMediaOutput.cpp
+++ b/Engine/Plugins/Media/PixelStreamingVCam/Source/PixelStreamingVCam/Private/PixelStreamingMediaOutput.cpp
@@ -28,7 +28,8 @@ UMediaCapture* UPixelStreamingMediaOutput::CreateMediaCaptureImpl()
 	{
 		Capture = NewObject<UPixelStreamingMediaCapture>();
 		Capture->SetMediaOutput(this);
-		Capture->OnStateChangedNative.AddLambda([this]() { OnCaptureStateChanged(); });
+		Capture->OnStateChangedNative.AddUObject(this, &UPixelStreamingMediaOutput::OnCaptureStateChanged);
+		Capture->OnCaptureViewportInitialized.AddUObject(this, &UPixelStreamingMediaOutput::OnCaptureViewportInitialized);
 	}
 	return Capture;
 }
@@ -49,14 +50,21 @@ void UPixelStreamingMediaOutput::OnCaptureStateChanged()
 	}
 }
 
+void UPixelStreamingMediaOutput::OnCaptureViewportInitialized()
+{
+	if(Streamer)
+	{
+		Streamer->SetTargetViewport(Capture->GetViewport().Get());
+		Streamer->SetTargetWindow(Capture->GetViewport()->FindWindow());
+	}
+}
+
 void UPixelStreamingMediaOutput::StartStreaming()
 {
 	if (Streamer)
 	{
 		Streamer->SetSignallingServerURL(SignallingServerURL);
 		Streamer->SetVideoInput(Capture->GetVideoInput());
-		Streamer->SetTargetViewport(Capture->GetViewport().Get());
-		Streamer->SetTargetWindow(Capture->GetViewport()->FindWindow());
 
 		if(!Streamer->IsStreaming())
 		{
diff --git a/Engine/Plugins/Media/PixelStreamingVCam/Source/PixelStreamingVCam/Private/PixelStreamingMediaOutput.h b/Engine/Plugins/Media/PixelStreamingVCam/Source/PixelStreamingVCam/Private/PixelStreamingMediaOutput.h
--- a/Engine/Plugins/Media/PixelStreamingVCam/Source/PixelStreamingVCam/Private/PixelStreamingMediaOutput.h
+++ b/Engine/Plugins/Media/PixelStreamingVCam/Source/PixelStreamingVCam/Private/PixelStreamingMediaOutput.h
@@ -39,6 +39,7 @@ private:
 	FString StreamID = TEXT("VCam");
 
 	void OnCaptureStateChanged();
+	void OnCaptureViewportInitialized();
 	void StartStreaming();
 	void StopStreaming();
 };
diff --git a/Engine/Plugins/Media/PixelStreamingVCam/Source/PixelStreamingVCam/Private/VCamPixelStreamingSession.cpp b/Engine/Plugins/Media/PixelStreamingVCam/Source/PixelStreamingVCam/Private/VCamPixelStreamingSession.cpp
--- a/Engine/Plugins/Media/PixelStreamingVCam/Source/PixelStreamingVCam/Private/VCamPixelStreamingSession.cpp
+++ b/Engine/Plugins/Media/PixelStreamingVCam/Source/PixelStreamingVCam/Private/VCamPixelStreamingSession.cpp
@@ -246,13 +246,6 @@ void UVCamPixelStreamingSession::PostEditChangeProperty(FPropertyChangedEvent& P
 		{
 			SetActive(false);
 		}
-		else if(Property->GetFName() == FName(TEXT("bIsActive")))
-		{
-			if(bIsActive) {
-				// this is a hack to suppress double activation, livelink tick will activate it
-				return;
-			}
-		}
 	}
 	Super::PostEditChangeProperty(PropertyChangedEvent);
 }
