Fix merge artifacts in PixelStreaminVCam and PixelCapture

diff --git a/Engine/Plugins/Media/PixelCapture/Source/PixelCapture/Public/PixelCaptureCapturer.h b/Engine/Plugins/Media/PixelCapture/Source/PixelCapture/Public/PixelCaptureCapturer.h
--- a/Engine/Plugins/Media/PixelCapture/Source/PixelCapture/Public/PixelCaptureCapturer.h
+++ b/Engine/Plugins/Media/PixelCapture/Source/PixelCapture/Public/PixelCaptureCapturer.h
@@ -4,7 +4,11 @@
 
 #include "IPixelCaptureInputFrame.h"
 #include "IPixelCaptureOutputFrame.h"
-#include "PixelCaptureFwd.h"
+
+namespace UE::PixelCapture
+{
+	class FOutputFrameBuffer;
+}
 
 /**
  * The base class for all Capturers in the system.
diff --git a/Engine/Plugins/Media/PixelCapture/Source/PixelCapture/Public/PixelCaptureFwd.h b/Engine/Plugins/Media/PixelCapture/Source/PixelCapture/Public/PixelCaptureFwd.h
deleted file mode 100644
diff --git a/Engine/Plugins/Media/PixelStreamingVCam/Source/PixelStreamingVCam/Private/VCamPixelStreamingSession.cpp b/Engine/Plugins/Media/PixelStreamingVCam/Source/PixelStreamingVCam/Private/VCamPixelStreamingSession.cpp
--- a/Engine/Plugins/Media/PixelStreamingVCam/Source/PixelStreamingVCam/Private/VCamPixelStreamingSession.cpp
+++ b/Engine/Plugins/Media/PixelStreamingVCam/Source/PixelStreamingVCam/Private/VCamPixelStreamingSession.cpp
@@ -146,6 +146,10 @@ void UVCamPixelStreamingSession::Activate()
 			double Timestamp;
 			Ar << Timestamp;
 
+			if(TSharedPtr<FPixelStreamingLiveLinkSource> LiveLinkSource = UVCamPixelStreamingSubsystem::Get()->LiveLinkSource)
+			{
+				LiveLinkSource->PushTransformForSubject(GetFName(), FTransform(ARKitMatrix), Timestamp);
+			}
 
 			/**
 			 * Code to control the level editor viewport
