Fix UMG HUD of PixelStreamingVCam output provider drawing underneath UE sprites/toolbars.

diff --git a/Engine/Plugins/Media/PixelStreamingVCam/Source/PixelStreamingVCam/Private/VCamPixelStreamingSession.cpp b/Engine/Plugins/Media/PixelStreamingVCam/Source/PixelStreamingVCam/Private/VCamPixelStreamingSession.cpp
--- a/Engine/Plugins/Media/PixelStreamingVCam/Source/PixelStreamingVCam/Private/VCamPixelStreamingSession.cpp
+++ b/Engine/Plugins/Media/PixelStreamingVCam/Source/PixelStreamingVCam/Private/VCamPixelStreamingSession.cpp
@@ -25,6 +25,7 @@ namespace VCamPixelStreamingSession
 
 void UVCamPixelStreamingSession::Initialize()
 {
+	DisplayType = EVPWidgetDisplayType::Viewport;
 	Super::Initialize();
 }
 
