Revert DisplayType = EVPWidgetDisplayType::Viewport to EVPWidgetDisplayType::PostProcess as HUD was no longer being streamed. To fix UI elements and toolbar drawing on top of HUD the solution is likely in the MediaCapture.h or VPFullScreenUserWidget.h but this is not my area of expertise.

diff --git a/Engine/Plugins/Media/PixelStreamingVCam/Source/PixelStreamingVCam/Private/VCamPixelStreamingSession.cpp b/Engine/Plugins/Media/PixelStreamingVCam/Source/PixelStreamingVCam/Private/VCamPixelStreamingSession.cpp
--- a/Engine/Plugins/Media/PixelStreamingVCam/Source/PixelStreamingVCam/Private/VCamPixelStreamingSession.cpp
+++ b/Engine/Plugins/Media/PixelStreamingVCam/Source/PixelStreamingVCam/Private/VCamPixelStreamingSession.cpp
@@ -25,7 +25,7 @@ namespace VCamPixelStreamingSession
 
 void UVCamPixelStreamingSession::Initialize()
 {
-	DisplayType = EVPWidgetDisplayType::Viewport;
+	DisplayType = EVPWidgetDisplayType::PostProcess;
 	Super::Initialize();
 }
 
