The LiveLinkSource for the VCam Pixel Streaming Subsystem will now attempt to recreate itself if it does not currently exist

diff --git a/Engine/Plugins/Media/PixelStreamingVCam/Source/PixelStreamingVCam/Private/VCamPixelStreamingSession.cpp b/Engine/Plugins/Media/PixelStreamingVCam/Source/PixelStreamingVCam/Private/VCamPixelStreamingSession.cpp
--- a/Engine/Plugins/Media/PixelStreamingVCam/Source/PixelStreamingVCam/Private/VCamPixelStreamingSession.cpp
+++ b/Engine/Plugins/Media/PixelStreamingVCam/Source/PixelStreamingVCam/Private/VCamPixelStreamingSession.cpp
@@ -147,7 +147,7 @@ void UVCamPixelStreamingSession::Activate()
 			double Timestamp;
 			Ar << Timestamp;
 
-			if(TSharedPtr<FPixelStreamingLiveLinkSource> LiveLinkSource = UVCamPixelStreamingSubsystem::Get()->LiveLinkSource)
+			if(TSharedPtr<FPixelStreamingLiveLinkSource> LiveLinkSource = UVCamPixelStreamingSubsystem::Get()->TryGetLiveLinkSource(this))
 			{
 				LiveLinkSource->PushTransformForSubject(GetFName(), FTransform(ARKitMatrix), Timestamp);
 			}
diff --git a/Engine/Plugins/Media/PixelStreamingVCam/Source/PixelStreamingVCam/Private/VCamPixelStreamingSubsystem.cpp b/Engine/Plugins/Media/PixelStreamingVCam/Source/PixelStreamingVCam/Private/VCamPixelStreamingSubsystem.cpp
--- a/Engine/Plugins/Media/PixelStreamingVCam/Source/PixelStreamingVCam/Private/VCamPixelStreamingSubsystem.cpp
+++ b/Engine/Plugins/Media/PixelStreamingVCam/Source/PixelStreamingVCam/Private/VCamPixelStreamingSubsystem.cpp
@@ -9,13 +9,6 @@
 void UVCamPixelStreamingSubsystem::Initialize(FSubsystemCollectionBase& Collection)
 {
 	Super::Initialize(Collection);
-	IModularFeatures& ModularFeatures = IModularFeatures::Get();
-	if (ModularFeatures.IsModularFeatureAvailable(ILiveLinkClient::ModularFeatureName))
-	{
-		ILiveLinkClient* LiveLinkClient = &ModularFeatures.GetModularFeature<ILiveLinkClient>(ILiveLinkClient::ModularFeatureName);
-		LiveLinkSource = MakeShared<FPixelStreamingLiveLinkSource>();
-		LiveLinkClient->AddSource(LiveLinkSource);
-	}
 }
 
 void UVCamPixelStreamingSubsystem::Deinitialize()
@@ -56,6 +49,43 @@ void UVCamPixelStreamingSubsystem::UnregisterActiveOutputProvider(UVCamPixelStre
 	}
 }
 
+TSharedPtr<FPixelStreamingLiveLinkSource> UVCamPixelStreamingSubsystem::TryGetLiveLinkSource()
+{
+	if (!LiveLinkSource.IsValid())
+	{
+		IModularFeatures& ModularFeatures = IModularFeatures::Get();
+		if (ModularFeatures.IsModularFeatureAvailable(ILiveLinkClient::ModularFeatureName))
+		{
+			ILiveLinkClient* LiveLinkClient = &ModularFeatures.GetModularFeature<ILiveLinkClient>(ILiveLinkClient::ModularFeatureName);
+			LiveLinkSource = MakeShared<FPixelStreamingLiveLinkSource>();
+			LiveLinkClient->AddSource(LiveLinkSource);
+		}
+	}
+	return LiveLinkSource;
+}
+
+TSharedPtr<FPixelStreamingLiveLinkSource> UVCamPixelStreamingSubsystem::TryGetLiveLinkSource(UVCamPixelStreamingSession* OutputProvider)
+{
+	if (!LiveLinkSource.IsValid())
+	{
+		IModularFeatures& ModularFeatures = IModularFeatures::Get();
+		if (ModularFeatures.IsModularFeatureAvailable(ILiveLinkClient::ModularFeatureName))
+		{
+			ILiveLinkClient* LiveLinkClient = &ModularFeatures.GetModularFeature<ILiveLinkClient>(ILiveLinkClient::ModularFeatureName);
+			LiveLinkSource = MakeShared<FPixelStreamingLiveLinkSource>();
+			LiveLinkClient->AddSource(LiveLinkSource);
+
+			if (IsValid(OutputProvider))
+			{
+				ActiveOutputProviders.AddUnique(OutputProvider);
+				LiveLinkSource->CreateSubject(OutputProvider->GetFName());
+				LiveLinkSource->PushTransformForSubject(OutputProvider->GetFName(), FTransform::Identity);
+			}
+		}
+	}
+	return LiveLinkSource;
+}
+
 TSharedPtr<UE::PixelStreamingServers::IServer> UVCamPixelStreamingSubsystem::LaunchSignallingServer(int StreamerPort, int PlayerPort)
 {
 	bool bAlreadyLaunched = Server.IsValid() && Server->HasLaunched();
diff --git a/Engine/Plugins/Media/PixelStreamingVCam/Source/PixelStreamingVCam/Public/VCamPixelStreamingSubsystem.h b/Engine/Plugins/Media/PixelStreamingVCam/Source/PixelStreamingVCam/Public/VCamPixelStreamingSubsystem.h
--- a/Engine/Plugins/Media/PixelStreamingVCam/Source/PixelStreamingVCam/Public/VCamPixelStreamingSubsystem.h
+++ b/Engine/Plugins/Media/PixelStreamingVCam/Source/PixelStreamingVCam/Public/VCamPixelStreamingSubsystem.h
@@ -21,11 +21,16 @@ public:
 	void RegisterActiveOutputProvider(UVCamPixelStreamingSession* OutputProvider);
 	void UnregisterActiveOutputProvider(UVCamPixelStreamingSession* OutputProvider);
 
+	// Get the LiveLinkSource if it already exists or attempt to create one
+	TSharedPtr<FPixelStreamingLiveLinkSource> TryGetLiveLinkSource();
+	
+	// Get the LiveLinkSource if it already exists or attempt to create one
+	// Additionally registers the OutputProvider to the Source if a new source was created
+	TSharedPtr<FPixelStreamingLiveLinkSource> TryGetLiveLinkSource(UVCamPixelStreamingSession* OutputProvider);
+
 	TSharedPtr<UE::PixelStreamingServers::IServer> LaunchSignallingServer(int StreamerPort, int PlayerPort);
 	void StopSignallingServer();
 	
-	// An associated Live Link Source shared by all output providers
-	TSharedPtr<FPixelStreamingLiveLinkSource> LiveLinkSource;
 private:
 	// Keep track of which output providers are currently active
 	UPROPERTY(Transient)
@@ -36,4 +41,7 @@ private:
 
 	// Download process for PS web frontend files (if we want to view output in the browser)
 	TSharedPtr<FMonitoredProcess> DownloadProcess;
+
+	// An associated Live Link Source shared by all output providers
+	TSharedPtr<FPixelStreamingLiveLinkSource> LiveLinkSource;
 };
