Added an option to pixel streaming output provider to automatically set the live link subject for the owning vcam component

diff --git a/Engine/Plugins/Media/PixelStreamingVCam/Source/PixelStreamingVCam/Private/VCamPixelStreamingSession.cpp b/Engine/Plugins/Media/PixelStreamingVCam/Source/PixelStreamingVCam/Private/VCamPixelStreamingSession.cpp
--- a/Engine/Plugins/Media/PixelStreamingVCam/Source/PixelStreamingVCam/Private/VCamPixelStreamingSession.cpp
+++ b/Engine/Plugins/Media/PixelStreamingVCam/Source/PixelStreamingVCam/Private/VCamPixelStreamingSession.cpp
@@ -16,6 +16,7 @@
 #include "PixelStreamingVCamLog.h"
 #include "IPixelStreamingModule.h"
 #include "PixelStreamingProtocol.h"
+#include "VCamComponent.h"
 
 namespace VCamPixelStreamingSession
 {
@@ -44,6 +45,11 @@ void UVCamPixelStreamingSession::Activate()
 	if (UVCamPixelStreamingSubsystem* PixelStreamingSubsystem = UVCamPixelStreamingSubsystem::Get())
 	{
 		PixelStreamingSubsystem->RegisterActiveOutputProvider(this);
+		UVCamComponent* VCamComponent = GetTypedOuter<UVCamComponent>();
+		if (bAutoSetLiveLinkSubject && IsValid(VCamComponent))
+		{
+			VCamComponent->LiveLinkSubject = GetFName();
+		}
 	}
 
 	// If we don't have a UMG assigned, we still need to create an empty 'dummy' UMG in order to properly route the input back from the RemoteSession device
diff --git a/Engine/Plugins/Media/PixelStreamingVCam/Source/PixelStreamingVCam/Public/VCamPixelStreamingSession.h b/Engine/Plugins/Media/PixelStreamingVCam/Source/PixelStreamingVCam/Public/VCamPixelStreamingSession.h
--- a/Engine/Plugins/Media/PixelStreamingVCam/Source/PixelStreamingVCam/Public/VCamPixelStreamingSession.h
+++ b/Engine/Plugins/Media/PixelStreamingVCam/Source/PixelStreamingVCam/Public/VCamPixelStreamingSession.h
@@ -53,6 +53,10 @@ public:
 	UPROPERTY(EditAnywhere, BlueprintReadWrite, Category = "Output")
 	bool PreventEditorIdle = true;
 
+	// If true then the Live Link Subject of the owning VCam Component will be set to the subject created by this Output Provider when the Provider is enabled
+	UPROPERTY(EditAnywhere, BlueprintReadWrite, Category = "Output")
+	bool bAutoSetLiveLinkSubject = true;
+	
 protected:
 	UPROPERTY(Transient)
 	TObjectPtr<UPixelStreamingMediaOutput> MediaOutput = nullptr;
