Fix: Touch input not scaling and interacting with windows which weren't targeted when using Virtual Cameras. This has been achieved by overriding the handling of touch events (Start, Move, End) in the `VCamPixelStreamingSession`

diff --git a/Engine/Plugins/Media/PixelStreamingVCam/Source/PixelStreamingVCam/PixelStreamingVCam.Build.cs b/Engine/Plugins/Media/PixelStreamingVCam/Source/PixelStreamingVCam/PixelStreamingVCam.Build.cs
--- a/Engine/Plugins/Media/PixelStreamingVCam/Source/PixelStreamingVCam/PixelStreamingVCam.Build.cs
+++ b/Engine/Plugins/Media/PixelStreamingVCam/Source/PixelStreamingVCam/PixelStreamingVCam.Build.cs
@@ -32,7 +32,11 @@ namespace UnrealBuildTool.Rules
 				"CinematicCamera",
 				"LiveLinkInterface",
 				"PixelCapture",
-				"PixelStreamingEditor"
+				"PixelStreamingEditor",
+				"SlateCore",
+				"Slate",
+				"VPUtilities",
+				"InputCore"
 			});
 
 			// Can't package non-editor targets (e.g. games) with UnrealEd, so this dependency should only be added in editor.
diff --git a/Engine/Plugins/Media/PixelStreamingVCam/Source/PixelStreamingVCam/Private/PixelStreamingMediaOutput.cpp b/Engine/Plugins/Media/PixelStreamingVCam/Source/PixelStreamingVCam/Private/PixelStreamingMediaOutput.cpp
--- a/Engine/Plugins/Media/PixelStreamingVCam/Source/PixelStreamingVCam/Private/PixelStreamingMediaOutput.cpp
+++ b/Engine/Plugins/Media/PixelStreamingVCam/Source/PixelStreamingVCam/Private/PixelStreamingMediaOutput.cpp
@@ -56,7 +56,6 @@ void UPixelStreamingMediaOutput::OnCaptureViewportInitialized()
 	if(Streamer)
 	{
 		Streamer->SetTargetViewport(Capture->GetViewport()->GetViewportWidget());
-		Streamer->SetTargetWindow(Capture->GetViewport()->FindWindow());
 	}
 }
 
diff --git a/Engine/Plugins/Media/PixelStreamingVCam/Source/PixelStreamingVCam/Private/VCamPixelStreamingSession.cpp b/Engine/Plugins/Media/PixelStreamingVCam/Source/PixelStreamingVCam/Private/VCamPixelStreamingSession.cpp
--- a/Engine/Plugins/Media/PixelStreamingVCam/Source/PixelStreamingVCam/Private/VCamPixelStreamingSession.cpp
+++ b/Engine/Plugins/Media/PixelStreamingVCam/Source/PixelStreamingVCam/Private/VCamPixelStreamingSession.cpp
@@ -18,6 +18,10 @@
 #include "PixelStreamingProtocol.h"
 #include "PixelStreamingEditorModule.h"
 #include "VCamComponent.h"
+#include "Input/HittestGrid.h"
+#include "Widgets/SVirtualWindow.h"
+#include "VPFullScreenUserWidget.h"
+#include "InputCoreTypes.h"
 
 namespace VCamPixelStreamingSession
 {
@@ -53,7 +57,7 @@ void UVCamPixelStreamingSession::Activate()
 		}
 	}
 
-	// If we don't have a UMG assigned, we still need to create an empty 'dummy' UMG in order to properly route the input back from the RemoteSession device
+	// If we don't have a UMG assigned, we still need to create an empty 'dummy' UMG in order to properly route the input back from the device
 	if (UMGClass == nullptr)
 	{
 		bUsingDummyUMG = true;
@@ -129,14 +133,56 @@ void UVCamPixelStreamingSession::Activate()
 		}
 	}
 
+	
+	Super::Activate();
+	// Super::Activate() creates our UMG, so we must do our stuff after
+	// If we have a UMG, then use it
+	if (UMGWidget)
+	{
+		TSharedPtr<SVirtualWindow> InputWindow;
+
+		// If we are rendering from a ComposureOutputProvider, we need to get the InputWindow from that UMG, not the one in the PixelStreamingOutputProvider
+		if (UVCamOutputComposure* ComposureProvider = Cast<UVCamOutputComposure>(GetOtherOutputProviderByIndex(FromComposureOutputProviderIndex)))
+		{
+			if (UVPFullScreenUserWidget* ComposureUMGWidget = ComposureProvider->GetUMGWidget())
+			{
+				InputWindow = ComposureUMGWidget->PostProcessDisplayType.GetSlateWindow();
+				UE_LOG(LogPixelStreamingVCam, Log, TEXT("InputChannel callback - Routing input to active viewport with Composure UMG"));
+			}
+			else
+			{
+				UE_LOG(LogPixelStreamingVCam, Warning, TEXT("InputChannel callback - Composure usage was requested, but the specified ComposureOutputProvider has no UMG set"));
+			}
+		}
+		else
+		{
+			InputWindow = UMGWidget->PostProcessDisplayType.GetSlateWindow();
+			UE_LOG(LogPixelStreamingVCam, Log, TEXT("InputChannel callback - Routing input to active viewport with UMG"));
+		}
+
+		MediaOutput->GetStreamer()->SetTargetWindow(InputWindow);
+		bRouteTouchMessageToWidget = true;
+	}
+	else
+	{
+		MediaOutput->GetStreamer()->SetTargetWindow(GetTargetInputWindow());
+		bRouteTouchMessageToWidget = true;
+		UE_LOG(LogPixelStreamingVCam, Log, TEXT("InputChannel callback - Routing input to active viewport"));
+	}
+
 	if(MediaOutput->GetStreamer())
 	{
 		IPixelStreamingModule& PixelStreamingModule = IPixelStreamingModule::Get();
-		
+		typedef Protocol::EPixelStreamingMessageTypes EType;		
 		Protocol::EPixelStreamingMessageDirection MessageDirection = Protocol::EPixelStreamingMessageDirection::ToStreamer;
+		TMap<FString, Protocol::FPixelStreamingInputMessage> Protocol = PixelStreamingModule.GetProtocol().ToStreamerProtocol;
 
-		typedef Protocol::EPixelStreamingMessageTypes EType;
-		Protocol::FPixelStreamingInputMessage Message = Protocol::FPixelStreamingInputMessage(100, 72, {
+		/*
+		* ====================
+		* ARKit Transform
+		* ====================
+		*/
+		Protocol::FPixelStreamingInputMessage ARKitMessage = Protocol::FPixelStreamingInputMessage(100, 72, {
 			// 4x4 Transform
 			EType::Float, EType::Float, EType::Float, EType::Float,
 			EType::Float, EType::Float, EType::Float, EType::Float,
@@ -146,7 +192,7 @@ void UVCamPixelStreamingSession::Activate()
 			EType::Double
 		});
 
-		const TFunction<void(FMemoryReader)>& Handler = [this](FMemoryReader Ar) { 
+		const TFunction<void(FMemoryReader)>& ARKitHandler = [this](FMemoryReader Ar) { 
 			// The buffer contains the transform matrix stored as 16 floats
 			FMatrix ARKitMatrix;
 			for (int32 Row = 0; Row < 4; ++Row)
@@ -168,23 +214,134 @@ void UVCamPixelStreamingSession::Activate()
 			{
 				LiveLinkSource->PushTransformForSubject(GetFName(), FTransform(ARKitMatrix), Timestamp);
 			}
+		};
+		PixelStreamingModule.RegisterMessage(MessageDirection, "ARKitTransform", ARKitMessage, ARKitHandler);
 
-			/**
-			 * Code to control the level editor viewport
-			 */
-			// for (FLevelEditorViewportClient* LevelVC : GEditor->GetLevelViewportClients())
-			// {
-			// 	if (LevelVC && LevelVC->IsPerspective())
-			// 	{
-			// 		LevelVC->SetViewLocation(Translation);
-			// 		LevelVC->SetViewRotation(ModifiedRotation);
-			// 	}
-			// }
+		if(!bRouteTouchMessageToWidget)
+		{
+			// The following code overrides touch input handling which is only necessary if we have a GUI and 
+			// want to route touch events to it's widget
+			return;
+		}
+
+		/*
+		* ====================
+		* TOUCH START
+		* ====================
+		*/
+		const TFunction<void(FMemoryReader)>& TouchStartHandler = [this](FMemoryReader Ar) { 
+			if(!MediaCapture->GetViewport().IsValid()) 
+			{
+				return;
+			}
+
+			uint8 NumTouches;
+			Ar << NumTouches;
+			for (uint8 TouchIdx = 0; TouchIdx < NumTouches; TouchIdx++)
+			{
+				uint16 TouchPosX, TouchPosY;
+				uint8 TouchIndex, TouchForce, TouchValid;
+				Ar << TouchPosX << TouchPosY << TouchIndex << TouchForce << TouchValid;
+
+				// If Touch is valid
+				if (TouchValid != 0)
+				{
+					//                                                                           convert range from 0,65536 -> 0,1
+					FVector2D Location = ConvertFromNormalizedScreenLocation(FVector2D(TouchPosX / (float) UINT16_MAX, TouchPosY / (float) UINT16_MAX));
+					Location = Location - MediaCapture->GetViewport()->GetCachedGeometry().GetAbsolutePosition();
+
+					FWidgetPath WidgetPath = FindRoutingMessageWidget(Location);
+					
+					if (WidgetPath.IsValid())
+					{
+						FScopedSwitchWorldHack SwitchWorld(WidgetPath);
+						FPointerEvent PointerEvent(0, TouchIndex, Location, Location, TouchForce / 255.0f, true);
+						FSlateApplication::Get().RoutePointerDownEvent(WidgetPath, PointerEvent);
+					}
+				}
+			}
+		};
+		PixelStreamingModule.RegisterMessage(MessageDirection, "TouchStart", *Protocol.Find("TouchStart"), TouchStartHandler);
+
+		/*
+		* ====================
+		* TOUCH MOVE
+		* ====================
+		*/
+		const TFunction<void(FMemoryReader)>& TouchMoveHandler = [this](FMemoryReader Ar) { 
+			if(!MediaCapture->GetViewport().IsValid()) 
+			{
+				return;
+			}
+			uint8 NumTouches;
+			Ar << NumTouches;
+			for (uint8 TouchIdx = 0; TouchIdx < NumTouches; TouchIdx++)
+			{
+				uint16 TouchPosX, TouchPosY;
+				uint8 TouchIndex, TouchForce, TouchValid;
+				Ar << TouchPosX << TouchPosY << TouchIndex << TouchForce << TouchValid;
+
+				// If Touch is valid
+				if (TouchValid != 0)
+				{
+					//                                                                           convert range from 0,65536 -> 0,1
+					FVector2D Location = ConvertFromNormalizedScreenLocation(FVector2D(TouchPosX / (float) UINT16_MAX, TouchPosY / (float) UINT16_MAX));
+					Location = Location - MediaCapture->GetViewport()->GetCachedGeometry().GetAbsolutePosition();
+
+					FWidgetPath WidgetPath = FindRoutingMessageWidget(Location);
+					
+					if (WidgetPath.IsValid())
+					{
+						FScopedSwitchWorldHack SwitchWorld(WidgetPath);
+						FPointerEvent PointerEvent(0, TouchIndex, Location, LastTouchLocation, TouchForce / 255.0f, true);
+						FSlateApplication::Get().RoutePointerMoveEvent(WidgetPath, PointerEvent, false);
+					}
+
+					LastTouchLocation = Location;
+				}
+			}
 		};
+		PixelStreamingModule.RegisterMessage(MessageDirection, "TouchMove", *Protocol.Find("TouchMove"), TouchMoveHandler);
+
+		/*
+		* ====================
+		* TOUCH END
+		* ====================
+		*/
+		const TFunction<void(FMemoryReader)>& TouchEndHandler = [this](FMemoryReader Ar) { 
+			if(!MediaCapture->GetViewport().IsValid()) 
+			{
+				return;
+			}
+
+			uint8 NumTouches;
+			Ar << NumTouches;
+			for (uint8 TouchIdx = 0; TouchIdx < NumTouches; TouchIdx++)
+			{
+				uint16 TouchPosX, TouchPosY;
+				uint8 TouchIndex, TouchForce, TouchValid;
+				Ar << TouchPosX << TouchPosY << TouchIndex << TouchForce << TouchValid;
 
-		PixelStreamingModule.RegisterMessage(MessageDirection, "ARKitTransform", Message, Handler);
+				// If Touch is valid
+				if (TouchValid != 0)
+				{
+					//                                                                           convert range from 0,65536 -> 0,1
+					FVector2D Location = ConvertFromNormalizedScreenLocation(FVector2D(TouchPosX / (float) UINT16_MAX, TouchPosY / (float) UINT16_MAX));
+					Location = Location - MediaCapture->GetViewport()->GetCachedGeometry().GetAbsolutePosition();
+
+					FWidgetPath WidgetPath = FindRoutingMessageWidget(Location);
+					
+					if (WidgetPath.IsValid())
+					{
+						FScopedSwitchWorldHack SwitchWorld(WidgetPath);
+						FPointerEvent PointerEvent(0, TouchIndex, Location, Location, 0.0f, true);
+						FSlateApplication::Get().RoutePointerUpEvent(WidgetPath, PointerEvent);
+					}
+				}
+			}
+		};
+		PixelStreamingModule.RegisterMessage(MessageDirection, "TouchEnd", *Protocol.Find("TouchEnd"), TouchEndHandler);
 	}
-	Super::Activate();
 }
 
 void UVCamPixelStreamingSession::StopSignallingServer()
@@ -257,4 +414,76 @@ void UVCamPixelStreamingSession::PostEditChangeProperty(FPropertyChangedEvent& P
 }
 #endif
 
+FWidgetPath UVCamPixelStreamingSession::FindRoutingMessageWidget(const FVector2D& Location) const
+{
+	if (TSharedPtr<SWindow> PlaybackWindowPinned = MediaOutput->GetStreamer()->GetTargetWindow().Pin())
+	{
+		if (PlaybackWindowPinned->AcceptsInput())
+		{
+			bool bIgnoreEnabledStatus = false;
+			TArray<FWidgetAndPointer> WidgetsAndCursors = PlaybackWindowPinned->GetHittestGrid().GetBubblePath(Location, FSlateApplication::Get().GetCursorRadius(), bIgnoreEnabledStatus);
+			return FWidgetPath(MoveTemp(WidgetsAndCursors));
+		}
+	}
+	return FWidgetPath();
+}
+
+FIntPoint UVCamPixelStreamingSession::ConvertFromNormalizedScreenLocation(const FVector2D& ScreenLocation)
+{
+	FIntPoint OutVector((int32)ScreenLocation.X, (int32)ScreenLocation.Y);
+
+	TSharedPtr<SWindow> ApplicationWindow = MediaCapture->GetViewport()->FindWindow();
+	if (ApplicationWindow.IsValid())
+	{
+		FVector2D WindowOrigin = ApplicationWindow->GetPositionInScreen();
+		TWeakPtr<SViewport> TargetViewport = MediaCapture->GetViewport()->GetViewportWidget();
+		if (TargetViewport.IsValid())
+			{
+			TSharedPtr<SViewport> ViewportWidget = TargetViewport.Pin();
+
+			if (ViewportWidget.IsValid())
+			{
+				FGeometry InnerWindowGeometry = ApplicationWindow->GetWindowGeometryInWindow();
+
+				// Find the widget path relative to the window
+				FArrangedChildren JustWindow(EVisibility::Visible);
+				JustWindow.AddWidget(FArrangedWidget(ApplicationWindow.ToSharedRef(), InnerWindowGeometry));
+
+				FWidgetPath PathToWidget(ApplicationWindow.ToSharedRef(), JustWindow);
+				if (PathToWidget.ExtendPathTo(FWidgetMatcher(ViewportWidget.ToSharedRef()), EVisibility::Visible))
+				{
+					FArrangedWidget ArrangedWidget = PathToWidget.FindArrangedWidget(ViewportWidget.ToSharedRef()).Get(FArrangedWidget::GetNullWidget());
+
+					FVector2D WindowClientOffset = ArrangedWidget.Geometry.GetAbsolutePosition();
+					FVector2D WindowClientSize = ArrangedWidget.Geometry.GetAbsoluteSize();
+
+					FVector2D OutTemp = WindowOrigin + WindowClientOffset + (ScreenLocation * WindowClientSize);
+					OutVector = FIntPoint((int32)OutTemp.X, (int32)OutTemp.Y);
+				}
+			}
+		}
+		else
+		{
+			FVector2D SizeInScreen = ApplicationWindow->GetSizeInScreen();
+			FVector2D OutTemp = SizeInScreen * ScreenLocation;
+			OutVector = FIntPoint((int32)OutTemp.X, (int32)OutTemp.Y);
+		}
+	}
+
+	return OutVector;
+}
+
+
+TWeakPtr<SWindow> UVCamPixelStreamingSession::GetTargetInputWindow() const
+{
+	TWeakPtr<SWindow> InputWindow;
+
+	if (UVCamComponent* OuterComponent = GetTypedOuter<UVCamComponent>())
+	{
+		InputWindow = OuterComponent->GetTargetInputWindow();
+	}
+
+	return InputWindow;
+}
+
 IMPLEMENT_MODULE(FDefaultModuleImpl, PixelStreamingVCam)
diff --git a/Engine/Plugins/Media/PixelStreamingVCam/Source/PixelStreamingVCam/Public/VCamPixelStreamingSession.h b/Engine/Plugins/Media/PixelStreamingVCam/Source/PixelStreamingVCam/Public/VCamPixelStreamingSession.h
--- a/Engine/Plugins/Media/PixelStreamingVCam/Source/PixelStreamingVCam/Public/VCamPixelStreamingSession.h
+++ b/Engine/Plugins/Media/PixelStreamingVCam/Source/PixelStreamingVCam/Public/VCamPixelStreamingSession.h
@@ -56,9 +56,14 @@ protected:
 private:
 	void StopSignallingServer();
 	void OnARKitTransformReceived(FPixelStreamingPlayerId PlayerId, uint8 Type, TArray<uint8> Data);
+	FWidgetPath FindRoutingMessageWidget(const FVector2D& Location) const;
+	FIntPoint ConvertFromNormalizedScreenLocation(const FVector2D& ScreenLocation);
+	TWeakPtr<SWindow> GetTargetInputWindow() const;
 
 private:
-	FHitResult LastViewportTouchResult;
-	bool bUsingDummyUMG = false;
-	bool bOldThrottleCPUWhenNotForeground;
+	FHitResult 	LastViewportTouchResult;
+	bool 		bUsingDummyUMG = false;
+	bool 		bOldThrottleCPUWhenNotForeground;
+	bool		bRouteTouchMessageToWidget = false;
+	FVector2D	LastTouchLocation = FVector2D(EForceInit::ForceInitToZero);
 };
