Updated VCamPixelStreamingSession to respect override resolutions

diff --git a/Engine/Plugins/Media/PixelStreamingVCam/Source/PixelStreamingVCam/Private/VCamPixelStreamingSession.cpp b/Engine/Plugins/Media/PixelStreamingVCam/Source/PixelStreamingVCam/Private/VCamPixelStreamingSession.cpp
--- a/Engine/Plugins/Media/PixelStreamingVCam/Source/PixelStreamingVCam/Private/VCamPixelStreamingSession.cpp
+++ b/Engine/Plugins/Media/PixelStreamingVCam/Source/PixelStreamingVCam/Private/VCamPixelStreamingSession.cpp
@@ -124,6 +124,11 @@ void UVCamPixelStreamingSession::Activate()
 		TWeakPtr<FSceneViewport> SceneViewport = GetTargetSceneViewport();
 		if (TSharedPtr<FSceneViewport> PinnedSceneViewport = SceneViewport.Pin())
 		{
+			// Apply the override resolution if applicable
+			if (bUseOverrideResolution)
+			{
+				PinnedSceneViewport->SetFixedViewportSize(OverrideResolution.X, OverrideResolution.Y);
+			}
 			MediaCapture->CaptureSceneViewport(PinnedSceneViewport, Options);
 			UE_LOG(LogPixelStreamingVCam, Log, TEXT("PixelStreaming set with viewport"));
 		}
@@ -222,6 +227,16 @@ void UVCamPixelStreamingSession::Deactivate()
 		StopSignallingServer();
 	}
 
+	// Remove the override resolution
+	if (bUseOverrideResolution)
+	{
+		TWeakPtr<FSceneViewport> SceneViewport = GetTargetSceneViewport();
+		if (TSharedPtr<FSceneViewport> PinnedSceneViewport = SceneViewport.Pin())
+		{
+			PinnedSceneViewport->SetFixedViewportSize(0, 0);
+		}
+	}
+
 	Super::Deactivate();
 	if (bUsingDummyUMG)
 	{
